Force PIX to use the active window when capturing via command.

diff --git a/Engine/Plugins/Developer/PixWinPlugin/Source/PixWinPlugin/Private/PixWinPluginModule.cpp b/Engine/Plugins/Developer/PixWinPlugin/Source/PixWinPlugin/Private/PixWinPluginModule.cpp
--- a/Engine/Plugins/Developer/PixWinPlugin/Source/PixWinPlugin/Private/PixWinPluginModule.cpp
+++ b/Engine/Plugins/Developer/PixWinPlugin/Source/PixWinPlugin/Private/PixWinPluginModule.cpp
@@ -13,6 +13,7 @@
 #include "Windows/AllowWindowsPlatformTypes.h"
 THIRD_PARTY_INCLUDES_START
 	#include <dxgi1_3.h>
+	#include <pix3.h>
 	#include <DXProgrammableCapture.h>
 THIRD_PARTY_INCLUDES_END
 #include "Windows/HideWindowsPlatformTypes.h"
@@ -66,6 +67,10 @@ namespace Impl
 		{
 #if PIX_PLUGIN_ENABLED
 			check(IsValid());
+
+			HWND WindowHandle = GetActiveWindow();
+			PIXSetTargetWindow(WindowHandle);
+
 			GraphicsAnalysis->BeginCapture();
 #endif
 		}
