Don't load WinPixGpuCapturer unless -attachPIX is enabled.

diff --git a/Engine/Plugins/Developer/PixWinPlugin/Source/PixWinPlugin/Private/PixWinPluginModule.cpp b/Engine/Plugins/Developer/PixWinPlugin/Source/PixWinPlugin/Private/PixWinPluginModule.cpp
--- a/Engine/Plugins/Developer/PixWinPlugin/Source/PixWinPlugin/Private/PixWinPluginModule.cpp
+++ b/Engine/Plugins/Developer/PixWinPlugin/Source/PixWinPlugin/Private/PixWinPluginModule.cpp
@@ -45,9 +45,7 @@ namespace Impl
 
 			if (!WinPixGpuCapturerHandle)
 			{
-#if !WITH_EDITOR
 				if (FParse::Param(FCommandLine::Get(), TEXT("attachPIX")))
-#endif
 				{
 					WinPixGpuCapturerHandle = PIXLoadLatestWinPixGpuCapturerLibrary();
 				}
