Fixing build when PIX is disabled.

diff --git a/Engine/Plugins/Developer/PixWinPlugin/Source/PixWinPlugin/Private/PixWinPluginModule.cpp b/Engine/Plugins/Developer/PixWinPlugin/Source/PixWinPlugin/Private/PixWinPluginModule.cpp
--- a/Engine/Plugins/Developer/PixWinPlugin/Source/PixWinPlugin/Private/PixWinPluginModule.cpp
+++ b/Engine/Plugins/Developer/PixWinPlugin/Source/PixWinPlugin/Private/PixWinPluginModule.cpp
@@ -63,12 +63,12 @@ namespace Impl
 			return WinPixGpuCapturerHandle != nullptr;
 		}
 
-		void BeginCapture(HWND WindowHandle, const FString& DestFileName)
+		void BeginCapture(void* WindowHandle, const FString& DestFileName)
 		{
 #if PIX_PLUGIN_ENABLED
 			if (WinPixGpuCapturerHandle)
 			{
-				PIXSetTargetWindow(WindowHandle);
+				PIXSetTargetWindow((HWND)WindowHandle);
 
 				PIXCaptureParameters Parameters{};
 				// UETODO: christopher.waters - implement capturing to a file.
@@ -252,12 +252,10 @@ void FPixWinPluginModule::BeginFrameCapture(void* HWnd, const FString& DestFileN
 {
 	UE_LOG(PixWinPlugin, Log, TEXT("Capturing a frame in PIX"));
 
-	HWND WindowHandle = (HWND)HWnd;
-
 	ENQUEUE_RENDER_COMMAND(StartRenderDocCapture)(
-		[this, WindowHandle, DestFileName](FRHICommandListImmediate& RHICmdList)
+		[this, HWnd, DestFileName](FRHICommandListImmediate& RHICmdList)
 		{
-			PixGraphicsAnalysisInterface->BeginCapture(WindowHandle, DestFileName);
+			PixGraphicsAnalysisInterface->BeginCapture(HWnd, DestFileName);
 		});
 }
 
