PixWinPlugin changes
- Adding PIX capture button which will open the capture in PIX if it wasn't attached.
- PIX frame captures weren't flushing at the correct time meaning we could end up with partial captures.
- Attaching PIX to the engine will now trigger FDynamicRHI::EnableIdealGPUCaptureOptions aka "r.RHISetGPUCaptureOptions 1"

diff --git a/Engine/Plugins/Developer/PixWinPlugin/Source/PixWinPlugin/PixWinPlugin.Build.cs b/Engine/Plugins/Developer/PixWinPlugin/Source/PixWinPlugin/PixWinPlugin.Build.cs
--- a/Engine/Plugins/Developer/PixWinPlugin/Source/PixWinPlugin/PixWinPlugin.Build.cs
+++ b/Engine/Plugins/Developer/PixWinPlugin/Source/PixWinPlugin/PixWinPlugin.Build.cs
@@ -1,29 +1,45 @@
 // Copyright Epic Games, Inc. All Rights Reserved.
 
+using EpicGames.Core;
+
 namespace UnrealBuildTool.Rules
 {
+	[SupportedPlatformGroups("Windows")]
 	public class PixWinPlugin : ModuleRules
 	{
 		public PixWinPlugin(ReadOnlyTargetRules Target) : base(Target)
         {
-			PrivateDependencyModuleNames.AddRange(new string[]
-			{
+			PublicDependencyModuleNames.AddAll(
+				"InputDevice",
+				"RenderCore"
+			);
+
+			PrivateDependencyModuleNames.AddAll(
 				"Core",
 				"CoreUObject",
 				"Engine",
-				"RenderCore",
-				"InputDevice",
-				"RHI",
-			});
+				"InputCore",
+				"Projects",
+				"RHI"
+			);
 
-			if (Target.Platform.IsInGroup(UnrealPlatformGroup.Windows))
+			if (Target.Configuration != UnrealTargetConfiguration.Shipping)
 			{
 				PrivateDependencyModuleNames.Add("WinPixEventRuntime");
 			}
 
 			if (Target.bBuildEditor == true)
 			{
-				PrivateDependencyModuleNames.Add("UnrealEd");
+				DynamicallyLoadedModuleNames.Add("LevelEditor");
+
+				PrivateDependencyModuleNames.AddAll(
+					"Slate"
+					, "SlateCore"
+					, "EditorFramework"
+					, "UnrealEd"
+					, "MainFrame"
+					, "GameProjectGeneration"
+				);
 			}
 		}
 	}
diff --git a/Engine/Plugins/Developer/PixWinPlugin/Source/PixWinPlugin/Private/PixWinPluginCommands.cpp b/Engine/Plugins/Developer/PixWinPlugin/Source/PixWinPlugin/Private/PixWinPluginCommands.cpp
new file mode 100644
--- /dev/null
+++ b/Engine/Plugins/Developer/PixWinPlugin/Source/PixWinPlugin/Private/PixWinPluginCommands.cpp
@@ -0,0 +1,16 @@
+// Copyright Epic Games, Inc. All Rights Reserved.
+
+#if WITH_EDITOR
+
+#include "PixWinPluginCommands.h"
+
+#define LOCTEXT_NAMESPACE "PixWinPlugin"
+
+void FPixWinPluginCommands::RegisterCommands()
+{
+	UI_COMMAND(CaptureFrameCommand, "Capture Frame", "Captures the next frame and launches the PIX UI if not attached (Alt+F12)", EUserInterfaceActionType::Button, FInputChord(EKeys::F12, EModifierKey::Alt));
+}
+
+#undef LOCTEXT_NAMESPACE
+
+#endif //WITH_EDITOR
diff --git a/Engine/Plugins/Developer/PixWinPlugin/Source/PixWinPlugin/Private/PixWinPluginCommands.h b/Engine/Plugins/Developer/PixWinPlugin/Source/PixWinPlugin/Private/PixWinPluginCommands.h
new file mode 100644
--- /dev/null
+++ b/Engine/Plugins/Developer/PixWinPlugin/Source/PixWinPlugin/Private/PixWinPluginCommands.h
@@ -0,0 +1,23 @@
+// Copyright Epic Games, Inc. All Rights Reserved.
+
+#pragma once
+
+#if WITH_EDITOR
+
+#include "PixWinPluginStyle.h"
+#include "Framework/Commands/Commands.h"
+
+class FPixWinPluginCommands : public TCommands<FPixWinPluginCommands>
+{
+public:
+	FPixWinPluginCommands()
+		: TCommands<FPixWinPluginCommands>(TEXT("PixWinPlugin"), NSLOCTEXT("Contexts", "PixWinPlugin", "PixWin Plugin"), NAME_None, FPixWinPluginStyle::Get()->GetStyleSetName())
+	{
+	}
+
+	virtual void RegisterCommands() override;
+
+	TSharedPtr<class FUICommandInfo> CaptureFrameCommand;
+};
+
+#endif
diff --git a/Engine/Plugins/Developer/PixWinPlugin/Source/PixWinPlugin/Private/PixWinPluginModule.cpp b/Engine/Plugins/Developer/PixWinPlugin/Source/PixWinPlugin/Private/PixWinPluginModule.cpp
--- a/Engine/Plugins/Developer/PixWinPlugin/Source/PixWinPlugin/Private/PixWinPluginModule.cpp
+++ b/Engine/Plugins/Developer/PixWinPlugin/Source/PixWinPlugin/Private/PixWinPluginModule.cpp
@@ -2,13 +2,19 @@
 
 #include "PixWinPluginModule.h"
 
-#include "CoreMinimal.h"
 #include "RenderingThread.h"
 #include "RHI.h"
 #include "UnrealClient.h"
 #include "Engine/GameViewportClient.h"
 #include "Engine/Engine.h"
 #include "Misc/CommandLine.h"
+#include "HAL/FileManager.h"
+
+#if WITH_EDITOR
+#include "Framework/Application/SlateApplication.h"
+#include "Editor.h"
+#include "SPixWinPluginEditorExtension.h"
+#endif
 
 #if !defined(WITH_PIX_EVENT_RUNTIME)
 #define WITH_PIX_EVENT_RUNTIME 0
@@ -25,23 +31,35 @@ THIRD_PARTY_INCLUDES_END
 #include "Windows/HideWindowsPlatformTypes.h"
 #endif
 
-#if WITH_EDITOR
-#include "Editor.h"
-#endif
-
-DEFINE_LOG_CATEGORY_STATIC(PixWinPlugin, Log, All);
+DEFINE_LOG_CATEGORY(PixWinPlugin);
 
 #define LOCTEXT_NAMESPACE "PixWinPlugin"
 
+static FString MakeWinPixCaptureFilePath(const FString& InFilename)
+{
+	FString Filename = InFilename;
+	if (Filename.IsEmpty())
+	{
+		const FDateTime DateTime = FDateTime::Now();
+		Filename = FString(TEXT("UEPixCapture_")) + DateTime.ToString();
+	}
+
+	const bool bAbsoluteFileName = !FPaths::IsRelative(Filename);
+	FString FileName = bAbsoluteFileName ? Filename : FPaths::ConvertRelativePathToFull(FPaths::ProjectSavedDir() / FString("PixCaptures") / Filename);
+	FileName = FPaths::SetExtension(FileName, "wpix");
+	FPaths::MakePlatformFilename(FileName);
+	return FileName;
+}
+
 namespace Impl
 {
+#if PIX_PLUGIN_ENABLED
 	/** Container for graphics analysis com interface. */
 	class FPixGraphicsAnalysisInterface
 	{
 	public:
 		FPixGraphicsAnalysisInterface()
 		{
-#if PIX_PLUGIN_ENABLED
 			WinPixGpuCapturerHandle = FPlatformProcess::GetDllHandle(L"WinPixGpuCapturer.dll");
 
 			if (!WinPixGpuCapturerHandle)
@@ -56,7 +74,6 @@ namespace Impl
 			{
 				PIXSetHUDOptions(PIX_HUD_SHOW_ON_NO_WINDOWS);
 			}
-#endif // PIX_PLUGIN_ENABLED
 		}
 
 		bool IsValid()
@@ -64,37 +81,46 @@ namespace Impl
 			return WinPixGpuCapturerHandle != nullptr;
 		}
 
-		void BeginCapture(void* WindowHandle, const FString& DestFileName)
+		bool IsAttached()
+		{
+			return WinPixGpuCapturerHandle != nullptr && PIXIsAttachedForGpuCapture();
+		}
+
+		void BeginCapture(HWND WindowHandle, const FString& DestFileName)
 		{
-#if PIX_PLUGIN_ENABLED
 			if (WinPixGpuCapturerHandle)
 			{
-				PIXSetTargetWindow((HWND)WindowHandle);
+				PIXSetTargetWindow(WindowHandle);
+
+				const FString CapturePath = MakeWinPixCaptureFilePath(DestFileName);
 
 				PIXCaptureParameters Parameters{};
-				// UETODO: christopher.waters - implement capturing to a file.
-				//Parameters.GpuCaptureParameters.FileName = *DestFileName;
-				Parameters.GpuCaptureParameters.FileName = TEXT("");
+				Parameters.GpuCaptureParameters.FileName = *CapturePath;
 
 				PIXBeginCapture2(PIX_CAPTURE_GPU, &Parameters);
 			}
-#endif
 		}
 
 		void EndCapture()
 		{
-#if PIX_PLUGIN_ENABLED
 			if (WinPixGpuCapturerHandle)
 			{
 				PIXEndCapture(0);
 			}
-#endif
+		}
+
+		void OpenCapture(const FString& FileName)
+		{
+			if (WinPixGpuCapturerHandle)
+			{
+				PIXOpenCaptureInUI(*FileName);
+			}
 		}
 
 	private:
 		void* WinPixGpuCapturerHandle{};
 	};
-
+#endif
 
 	/** Dummy input device that is used only to generate a Tick. */
 	class FPixDummyInputDevice : public IInputDevice
@@ -102,7 +128,8 @@ namespace Impl
 	public:
 		FPixDummyInputDevice(FPixWinPluginModule* InModule)
 			: Module(InModule)
-		{ }
+		{
+		}
 
 		virtual void Tick(float DeltaTime) override
 		{
@@ -125,9 +152,16 @@ namespace Impl
 
 void FPixWinPluginModule::StartupModule()
 {
+#if PIX_PLUGIN_ENABLED
 	PixGraphicsAnalysisInterface = new Impl::FPixGraphicsAnalysisInterface();
 	if (PixGraphicsAnalysisInterface->IsValid())
 	{
+		FString CapturePath = FPaths::ProjectSavedDir() / TEXT("PixCaptures");
+		if (!IFileManager::Get().DirectoryExists(*CapturePath))
+		{
+			IFileManager::Get().MakeDirectory(*CapturePath, true);
+		}
+
 		// Register modular features.
 		IModularFeatures::Get().RegisterModularFeature(IRenderCaptureProvider::GetModularFeatureName(), (IRenderCaptureProvider*)this);
 		IModularFeatures::Get().RegisterModularFeature(IInputDeviceModule::GetModularFeatureName(), (IInputDeviceModule*)this);
@@ -138,9 +172,12 @@ void FPixWinPluginModule::StartupModule()
 			TEXT("Captures the rendering commands of the next frame."),
 			FConsoleCommandDelegate::CreateRaw(this, &FPixWinPluginModule::CaptureFrame));
 
+		FCoreDelegates::OnPostEngineInit.AddRaw(this, &FPixWinPluginModule::OnPostEngineInit);
+
 		UE_LOG(PixWinPlugin, Log, TEXT("PIX capture plugin is ready!"));
 	}
 	else
+#endif
 	{
 		UE_LOG(PixWinPlugin, Log, TEXT("PIX capture plugin failed to initialize! Check that the process is launched from PIX."));
 	}
@@ -148,6 +185,7 @@ void FPixWinPluginModule::StartupModule()
 
 void FPixWinPluginModule::ShutdownModule()
 {
+#if PIX_PLUGIN_ENABLED
 	delete PixGraphicsAnalysisInterface;
 	PixGraphicsAnalysisInterface = nullptr;
 	delete ConsoleCommandCaptureFrame;
@@ -155,16 +193,21 @@ void FPixWinPluginModule::ShutdownModule()
 
 	IModularFeatures::Get().UnregisterModularFeature(IRenderCaptureProvider::GetModularFeatureName(), (IRenderCaptureProvider*)this);
 	IModularFeatures::Get().UnregisterModularFeature(IInputDeviceModule::GetModularFeatureName(), (IInputDeviceModule*)this);
+
+#if WITH_EDITOR
+	EditorExtension.Reset();
+#endif
+
+#endif
 }
 
-TSharedPtr<class IInputDevice> FPixWinPluginModule::CreateInputDevice(const TSharedRef<FGenericApplicationMessageHandler>& InMessageHandler)
+TSharedPtr<IInputDevice> FPixWinPluginModule::CreateInputDevice(const TSharedRef<FGenericApplicationMessageHandler>& InMessageHandler)
 {
 	UE_LOG(PixWinPlugin, Log, TEXT("Creating dummy input device (for intercepting engine ticks)"));
-	Impl::FPixDummyInputDevice* InputDev = new Impl::FPixDummyInputDevice(this);
-	return MakeShareable((IInputDevice*)InputDev);
+	return MakeShared<Impl::FPixDummyInputDevice>(this);
 }
 
-void FPixWinPluginModule::CaptureFrame(FViewport* InViewport, uint32 InFlags, FString const& InDestFileName)
+void FPixWinPluginModule::CaptureFrame(FViewport* InViewport, uint32 InFlags, const FString& InDestFileName)
 {
 	if (!bEndCaptureNextTick)
 	{
@@ -172,51 +215,74 @@ void FPixWinPluginModule::CaptureFrame(FViewport* InViewport, uint32 InFlags, FS
 	}
 }
 
-void FPixWinPluginModule::BeginCapture(FRHICommandListImmediate* InRHICommandList, uint32 InFlags, FString const& InDestFileName)
+void FPixWinPluginModule::BeginCapture(FRHICommandListImmediate* InRHICommandList, uint32 InFlags, const FString& InDestFileName)
 {
+#if PIX_PLUGIN_ENABLED
 	InRHICommandList->SubmitCommandsAndFlushGPU();
 	InRHICommandList->EnqueueLambda([Pix = PixGraphicsAnalysisInterface, InDestFileName](FRHICommandListImmediate& RHICommandList)
 	{
 		Pix->BeginCapture(nullptr, InDestFileName);
 	});
+#endif
 }
 
 void FPixWinPluginModule::EndCapture(FRHICommandListImmediate* InRHICommandList)
 {
+#if PIX_PLUGIN_ENABLED
 	InRHICommandList->SubmitCommandsAndFlushGPU();
 	InRHICommandList->EnqueueLambda([Pix = PixGraphicsAnalysisInterface](FRHICommandListImmediate& RHICommandList)
 	{
 		Pix->EndCapture();
 	});
+#endif
+}
+
+void FPixWinPluginModule::OnPostEngineInit()
+{
+#if WITH_EDITOR
+	if (FSlateApplication::IsInitialized() && !IsRunningCommandlet())
+	{
+		EditorExtension = MakeShared<FPixWinPluginEditorExtension>(this);
+	}
+#endif // WITH_EDITOR
 }
 
 void FPixWinPluginModule::Tick(float DeltaTime)
 {
+#if PIX_PLUGIN_ENABLED
+	const bool bNewCurrentlyAttached = PixGraphicsAnalysisInterface && PixGraphicsAnalysisInterface->IsAttached();
+	if (bNewCurrentlyAttached != bCurrentlyAttached)
+	{
+		// Call EnableIdealGPUCaptureOptions the first time we are attached
+		if (bNewCurrentlyAttached)
+		{
+			FDynamicRHI::EnableIdealGPUCaptureOptions(true);
+		}
+
+		bCurrentlyAttached = bNewCurrentlyAttached;
+	}
+
 	if (bBeginCaptureNextTick)
 	{
 		// Start a capture.
 		bBeginCaptureNextTick = false;
 		bEndCaptureNextTick = true;
 
-		ENQUEUE_RENDER_COMMAND(BeginCaptureCommand)([this](FRHICommandListImmediate& RHICommandList)
-		{
-			BeginCapture(&RHICommandList, 0, FString());
-		});
+		BeginFrameCapture(nullptr, FString());
 	}
 	else if (bEndCaptureNextTick)
 	{
 		// End a capture.
 		bEndCaptureNextTick = false;
 
-		ENQUEUE_RENDER_COMMAND(EndCaptureCommand)([this](FRHICommandListImmediate& RHICommandList)
-		{
-			EndCapture(&RHICommandList);
-		});
+		EndFrameCapture(0, FString());
 	}
+#endif
 }
 
-void FPixWinPluginModule::DoFrameCaptureCurrentViewport(FViewport* InViewport, uint32 InFlags, FString const& InDestFileName)
+void FPixWinPluginModule::DoFrameCaptureCurrentViewport(FViewport* InViewport, uint32 InFlags, const FString& InDestFileName)
 {
+#if PIX_PLUGIN_ENABLED
 	// infer the intended viewport to intercept/capture:
 	FViewport* Viewport = InViewport;
 
@@ -241,34 +307,59 @@ void FPixWinPluginModule::DoFrameCaptureCurrentViewport(FViewport* InViewport, u
 	}
 #endif // WITH_EDITOR
 
+	const FString DestFileName = MakeWinPixCaptureFilePath(InDestFileName);
+
 	check(Viewport);
-	BeginFrameCapture(Viewport->GetWindow(), InDestFileName);
+	BeginFrameCapture(Viewport->GetWindow(), DestFileName);
 
 	Viewport->Draw(true);
 
-	EndFrameCapture();
+	EndFrameCapture(InFlags, DestFileName);
+#endif
 }
 
 void FPixWinPluginModule::BeginFrameCapture(void* HWnd, const FString& DestFileName)
 {
+#if PIX_PLUGIN_ENABLED
 	UE_LOG(PixWinPlugin, Log, TEXT("Capturing a frame in PIX"));
 
-	ENQUEUE_RENDER_COMMAND(StartRenderDocCapture)(
-		[this, HWnd, DestFileName](FRHICommandListImmediate& RHICmdList)
+	Impl::FPixGraphicsAnalysisInterface* Pix = PixGraphicsAnalysisInterface;
+	HWND WindowHandle = HWnd ? (HWND)HWnd : GetActiveWindow();
+
+	ENQUEUE_RENDER_COMMAND(PixWinBeginFrameCapture)(
+		[Pix, WindowHandle, DestFileName](FRHICommandListImmediate& RHICommandList)
 		{
-			PixGraphicsAnalysisInterface->BeginCapture(HWnd, DestFileName);
+			if (Pix && Pix->IsValid())
+			{
+				Pix->BeginCapture(WindowHandle, FString());
+			}
 		});
+#endif
 }
 
-void FPixWinPluginModule::EndFrameCapture()
+void FPixWinPluginModule::EndFrameCapture(uint32 InFlags, const FString& DestFileName)
 {
-	ENQUEUE_RENDER_COMMAND(EndRenderDocCapture)(
-		[this](FRHICommandListImmediate& RHICmdList)
+#if PIX_PLUGIN_ENABLED
+	Impl::FPixGraphicsAnalysisInterface* Pix = PixGraphicsAnalysisInterface;
+	ENQUEUE_RENDER_COMMAND(PixWinEndFrameCapture)(
+		[Pix, InFlags, DestFileName](FRHICommandListImmediate& RHICommandList)
 		{
-			PixGraphicsAnalysisInterface->EndCapture();
+			if (Pix && Pix->IsValid())
+			{
+				RHICommandList.SubmitCommandsAndFlushGPU();
+				Pix->EndCapture();
+
+				// If we're already attached, don't open a new PIX instance
+				if ((InFlags & IRenderCaptureProvider::ECaptureFlags_Launch) && !Pix->IsAttached())
+				{
+					Pix->OpenCapture(DestFileName);
+				}
+			}
 		});
+#endif
 }
 
 #undef LOCTEXT_NAMESPACE
 
 IMPLEMENT_MODULE(FPixWinPluginModule, PixWinPlugin)
+
diff --git a/Engine/Plugins/Developer/PixWinPlugin/Source/PixWinPlugin/Private/PixWinPluginModule.h b/Engine/Plugins/Developer/PixWinPlugin/Source/PixWinPlugin/Private/PixWinPluginModule.h
--- a/Engine/Plugins/Developer/PixWinPlugin/Source/PixWinPlugin/Private/PixWinPluginModule.h
+++ b/Engine/Plugins/Developer/PixWinPlugin/Source/PixWinPlugin/Private/PixWinPluginModule.h
@@ -2,10 +2,15 @@
 
 #pragma once
 
-#include "CoreMinimal.h"
 #include "IPixWinPlugin.h"
+#include "Logging/LogMacros.h"
+
+DECLARE_LOG_CATEGORY_EXTERN(PixWinPlugin, Log, All);
 
 class FAutoConsoleCommand;
+class FPixWinPluginEditorExtension;
+class IInputDevice;
+
 namespace Impl { class FPixDummyInputDevice; }
 namespace Impl { class FPixGraphicsAnalysisInterface; }
 
@@ -16,8 +21,8 @@ class FPixWinPluginModule : public IPixWinPlugin
 
 public:
 	// Begin IRenderCaptureProvider interface.
-	virtual void CaptureFrame(FViewport* InViewport, uint32 InFlags, FString const& InDestFileName) override;
-	virtual void BeginCapture(FRHICommandListImmediate* InRHICommandList, uint32 InFlags, FString const& InDestFileName) override;
+	virtual void CaptureFrame(FViewport* InViewport, uint32 InFlags, const FString& InDestFileName) override;
+	virtual void BeginCapture(FRHICommandListImmediate* InRHICommandList, uint32 InFlags, const FString& InDestFileName) override;
 	virtual void EndCapture(FRHICommandListImmediate* InRHICommandList) override;
 	// End IRenderCaptureProvider interface.
 
@@ -28,17 +33,19 @@ protected:
 	// End IModuleInterface interface.
 
 	// Begin IInputDeviceModule interface.
-	virtual TSharedPtr<class IInputDevice> CreateInputDevice(const TSharedRef<FGenericApplicationMessageHandler>& InMessageHandler) override;
+	virtual TSharedPtr<IInputDevice> CreateInputDevice(const TSharedRef<FGenericApplicationMessageHandler>& InMessageHandler) override;
 	// End IInputDeviceModule interface.
 
 private:
+	void OnPostEngineInit();
+
 	/** Tick to handle full frame capture events. */
 	void Tick(float DeltaTime);
 
 	void DoFrameCaptureCurrentViewport(FViewport* InViewport, uint32 InFlags, const FString& InDestFileName);
 
 	void BeginFrameCapture(void* HWnd, const FString& DestFileName);
-	void EndFrameCapture();
+	void EndFrameCapture(uint32 InFlags, const FString& DestFileName);
 
 	/** Helper function for CVar command binding. */
 	void CaptureFrame() { return CaptureFrame(nullptr, 0, FString()); }
@@ -48,4 +55,10 @@ private:
 
 	bool bBeginCaptureNextTick = false;
 	bool bEndCaptureNextTick = false;
+
+	bool bCurrentlyAttached = false;
+
+#if WITH_EDITOR
+	TSharedPtr<FPixWinPluginEditorExtension> EditorExtension;
+#endif
 };
diff --git a/Engine/Plugins/Developer/PixWinPlugin/Source/PixWinPlugin/Private/PixWinPluginStyle.cpp b/Engine/Plugins/Developer/PixWinPlugin/Source/PixWinPlugin/Private/PixWinPluginStyle.cpp
new file mode 100644
--- /dev/null
+++ b/Engine/Plugins/Developer/PixWinPlugin/Source/PixWinPlugin/Private/PixWinPluginStyle.cpp
@@ -0,0 +1,62 @@
+// Copyright Epic Games, Inc. All Rights Reserved.
+
+#if WITH_EDITOR
+
+#include "PixWinPluginStyle.h"
+#include "Styling/SlateStyle.h"
+#include "Styling/SlateStyleRegistry.h"
+#include "HAL/FileManager.h"
+#include "Interfaces/IPluginManager.h"
+
+FString FPixWinPluginStyle::InContent(const FString& RelativePath, const ANSICHAR* Extension)
+{
+	TSharedPtr<IPlugin> ThisPlugin = IPluginManager::Get().FindPlugin(TEXT("PixWinPlugin"));
+	check(ThisPlugin.IsValid());
+	static FString ContentDir = ThisPlugin->GetBaseDir() / TEXT("Resources");
+	return (ContentDir / RelativePath) + Extension;
+}
+
+TSharedPtr<FSlateStyleSet> FPixWinPluginStyle::StyleSet = nullptr;
+TSharedPtr<class ISlateStyle> FPixWinPluginStyle::Get() { return StyleSet; }
+
+void FPixWinPluginStyle::Initialize()
+{
+	// Only register once
+	if (StyleSet.IsValid())
+	{
+		return;
+	}
+
+	StyleSet = MakeShareable(new FSlateStyleSet("PixWinPluginStyle"));
+
+	FString ProjectResourceDir = FPaths::ProjectPluginsDir() / TEXT("PixWinPlugin") / TEXT("Resources");
+	FString EngineResourceDir = FPaths::EnginePluginsDir() / TEXT("PixWinPlugin") / TEXT("Resources");
+
+	if (IFileManager::Get().DirectoryExists(*ProjectResourceDir)) //Is the plugin in the project? In that case, use those resources
+	{
+		StyleSet->SetContentRoot(ProjectResourceDir);
+		StyleSet->SetCoreContentRoot(ProjectResourceDir);
+	}
+	else //Otherwise, use the global ones
+	{
+		StyleSet->SetContentRoot(EngineResourceDir);
+		StyleSet->SetCoreContentRoot(EngineResourceDir);
+	}
+
+	StyleSet->Set("PixWinPlugin.Icon", new FSlateImageBrush(FPixWinPluginStyle::InContent("Icon40", ".png"), FVector2D(40.0f, 40.0f)));
+	StyleSet->Set("PixWinPlugin.CaptureFrameIcon", new FSlateImageBrush(FPixWinPluginStyle::InContent("ViewportIcon16", ".png"), FVector2D(16.0f, 16.0f)));
+
+	FSlateStyleRegistry::RegisterSlateStyle(*StyleSet.Get());
+};
+
+void FPixWinPluginStyle::Shutdown()
+{
+	if (StyleSet.IsValid())
+	{
+		FSlateStyleRegistry::UnRegisterSlateStyle(*StyleSet.Get());
+		ensure(StyleSet.IsUnique());
+		StyleSet.Reset();
+	}
+}
+
+#endif //WITH_EDITOR
diff --git a/Engine/Plugins/Developer/PixWinPlugin/Source/PixWinPlugin/Private/PixWinPluginStyle.h b/Engine/Plugins/Developer/PixWinPlugin/Source/PixWinPlugin/Private/PixWinPluginStyle.h
new file mode 100644
--- /dev/null
+++ b/Engine/Plugins/Developer/PixWinPlugin/Source/PixWinPlugin/Private/PixWinPluginStyle.h
@@ -0,0 +1,23 @@
+// Copyright Epic Games, Inc. All Rights Reserved.
+
+#pragma once
+
+#if WITH_EDITOR
+
+#include "Styling/ISlateStyle.h"
+
+class FPixWinPluginStyle
+{
+public:
+	static void Initialize();
+	static void Shutdown();
+	static TSharedPtr<class ISlateStyle> Get();
+
+private:
+	static FString InContent(const FString& RelativePath, const ANSICHAR* Extension);
+
+private:
+	static TSharedPtr<class FSlateStyleSet> StyleSet;
+};
+
+#endif
diff --git a/Engine/Plugins/Developer/PixWinPlugin/Source/PixWinPlugin/Private/SPixWinPluginEditorExtension.cpp b/Engine/Plugins/Developer/PixWinPlugin/Source/PixWinPlugin/Private/SPixWinPluginEditorExtension.cpp
new file mode 100644
--- /dev/null
+++ b/Engine/Plugins/Developer/PixWinPlugin/Source/PixWinPlugin/Private/SPixWinPluginEditorExtension.cpp
@@ -0,0 +1,144 @@
+// Copyright Epic Games, Inc. All Rights Reserved.
+
+#include "SPixWinPluginEditorExtension.h"
+
+#if WITH_EDITOR
+
+#include "PixWinPluginCommands.h"
+#include "PixWinPluginModule.h"
+#include "PixWinPluginStyle.h"
+
+#include "Editor/EditorEngine.h"
+#include "SEditorViewportToolBarMenu.h"
+#include "SViewportToolBarComboMenu.h"
+#include "Kismet2/DebuggerCommands.h"
+#include "Styling/AppStyle.h"
+#include "Framework/Application/SlateApplication.h"
+#include "Framework/MultiBox/MultiBoxBuilder.h"
+#include "Framework/MultiBox/MultiBoxDefs.h"
+#include "Interfaces/IMainFrameModule.h"
+#include "LevelEditor.h"
+#include "RHI.h"
+#include "Widgets/Input/SButton.h"
+#include "Widgets/Images/SImage.h"
+
+extern UNREALED_API UEditorEngine* GEditor;
+
+class SPixWinCaptureButton : public SViewportToolBar
+{
+public:
+	SLATE_BEGIN_ARGS(SPixWinCaptureButton) { }
+	SLATE_END_ARGS()
+
+	/** Widget constructor */
+	void Construct(const FArguments& Args)
+	{
+		FSlateIcon IconBrush = FSlateIcon(FPixWinPluginStyle::Get()->GetStyleSetName(), "PixWinPlugin.CaptureFrameIcon");
+
+		ChildSlot
+		[
+			SNew(SButton)
+			.HAlign(HAlign_Center)
+			.VAlign(VAlign_Center)
+			.ContentPadding(FMargin(1.0f))
+			.ToolTipText(FPixWinPluginCommands::Get().CaptureFrameCommand->GetDescription())
+			.OnClicked_Lambda([this]()
+			{
+				FPlayWorldCommands::GlobalPlayWorldActions->GetActionForCommand(FPixWinPluginCommands::Get().CaptureFrameCommand)->Execute();
+				return(FReply::Handled());
+			})
+			[
+				SNew(SImage)
+				.Image(IconBrush.GetIcon())
+			]
+		];
+	}
+};
+
+FPixWinPluginEditorExtension::FPixWinPluginEditorExtension(FPixWinPluginModule* ThePlugin)
+	: ToolbarExtension()
+	, ExtensionManager()
+	, ToolbarExtender()
+{
+	Initialize(ThePlugin);
+}
+
+FPixWinPluginEditorExtension::~FPixWinPluginEditorExtension()
+{
+	if (ExtensionManager.IsValid())
+	{
+		FPixWinPluginStyle::Shutdown();
+		FPixWinPluginCommands::Unregister();
+
+		ToolbarExtender->RemoveExtension(ToolbarExtension.ToSharedRef());
+
+		ExtensionManager->RemoveExtender(ToolbarExtender);
+	}
+	else
+	{
+		ExtensionManager.Reset();
+	}
+}
+
+void FPixWinPluginEditorExtension::Initialize(FPixWinPluginModule* ThePlugin)
+{
+	if (GUsingNullRHI)
+	{
+		UE_LOG(PixWinPlugin, Display, TEXT("PixWin Plugin will not be loaded because a Null RHI (Cook Server, perhaps) is being used."));
+		return;
+	}
+
+	// The LoadModule request below will crash if running as an editor commandlet!
+	check(!IsRunningCommandlet());
+
+	FPixWinPluginStyle::Initialize();
+	FPixWinPluginCommands::Register();
+
+#if WITH_EDITOR
+	if (!IsRunningGame())
+	{
+		FLevelEditorModule& LevelEditorModule = FModuleManager::LoadModuleChecked<FLevelEditorModule>("LevelEditor");
+		TSharedRef<FUICommandList> CommandBindings = LevelEditorModule.GetGlobalLevelEditorActions();
+		ExtensionManager = LevelEditorModule.GetToolBarExtensibilityManager();
+		ToolbarExtender = MakeShareable(new FExtender);
+		ToolbarExtension = ToolbarExtender->AddToolBarExtension("CameraSpeed", EExtensionHook::After, CommandBindings,
+			FToolBarExtensionDelegate::CreateLambda([this, ThePlugin](FToolBarBuilder& ToolbarBuilder)
+				{
+					AddToolbarExtension(ToolbarBuilder, ThePlugin);
+				})
+		);
+		ExtensionManager->AddExtender(ToolbarExtender);
+	}
+#endif // WITH_EDITOR
+
+	// Would be nice to use the preprocessor definition WITH_EDITOR instead, but the user may launch a standalone the game through the editor...
+	if (GEditor != nullptr)
+	{
+		check(FPlayWorldCommands::GlobalPlayWorldActions.IsValid());
+
+		//Register the editor hotkeys
+		FPlayWorldCommands::GlobalPlayWorldActions->MapAction(FPixWinPluginCommands::Get().CaptureFrameCommand,
+			FExecuteAction::CreateLambda([]()
+				{
+					FPixWinPluginModule& PluginModule = FModuleManager::GetModuleChecked<FPixWinPluginModule>("PixWinPlugin");
+					PluginModule.CaptureFrame(nullptr, IRenderCaptureProvider::ECaptureFlags_Launch, FString());
+				}),
+			FCanExecuteAction());
+	}
+}
+
+void FPixWinPluginEditorExtension::AddToolbarExtension(FToolBarBuilder& ToolbarBuilder, FPixWinPluginModule* ThePlugin)
+{
+#define LOCTEXT_NAMESPACE "LevelEditorToolBar"
+
+	UE_LOG(PixWinPlugin, Log, TEXT("Attaching toolbar extension..."));
+	ToolbarBuilder.AddSeparator();
+
+	ToolbarBuilder.BeginSection("PixWinPlugin");
+	ToolbarBuilder.AddWidget(SNew(SPixWinCaptureButton));
+	ToolbarBuilder.EndSection();
+
+#undef LOCTEXT_NAMESPACE
+}
+
+#endif //WITH_EDITOR
diff --git a/Engine/Plugins/Developer/PixWinPlugin/Source/PixWinPlugin/Private/SPixWinPluginEditorExtension.h b/Engine/Plugins/Developer/PixWinPlugin/Source/PixWinPlugin/Private/SPixWinPluginEditorExtension.h
new file mode 100644
--- /dev/null
+++ b/Engine/Plugins/Developer/PixWinPlugin/Source/PixWinPlugin/Private/SPixWinPluginEditorExtension.h
@@ -0,0 +1,30 @@
+// Copyright Epic Games, Inc. All Rights Reserved.
+
+#pragma once
+
+#if WITH_EDITOR
+
+#include "SViewportToolBar.h"
+
+class FPixWinPluginModule;
+class FExtensionBase;
+class FExtensibilityManager;
+class FExtender;
+class FToolBarBuilder;
+
+class FPixWinPluginEditorExtension
+{
+public:
+	FPixWinPluginEditorExtension(FPixWinPluginModule* ThePlugin);
+	~FPixWinPluginEditorExtension();
+
+private:
+	void Initialize(FPixWinPluginModule* ThePlugin);
+	void AddToolbarExtension(FToolBarBuilder& ToolbarBuilder, FPixWinPluginModule* ThePlugin);
+
+	TSharedPtr<const FExtensionBase> ToolbarExtension;
+	TSharedPtr<FExtensibilityManager> ExtensionManager;
+	TSharedPtr<FExtender> ToolbarExtender;
+};
+
+#endif //WITH_EDITOR
diff --git a/Engine/Plugins/Developer/PixWinPlugin/Source/PixWinPlugin/Public/IPixWinPlugin.h b/Engine/Plugins/Developer/PixWinPlugin/Source/PixWinPlugin/Public/IPixWinPlugin.h
--- a/Engine/Plugins/Developer/PixWinPlugin/Source/PixWinPlugin/Public/IPixWinPlugin.h
+++ b/Engine/Plugins/Developer/PixWinPlugin/Source/PixWinPlugin/Public/IPixWinPlugin.h
@@ -2,9 +2,9 @@
 
 #pragma once
 
-#include "CoreMinimal.h"
 #include "IInputDeviceModule.h"
 #include "IRenderCaptureProvider.h"
+#include "Modules/ModuleManager.h"
 
 /** PIX capture plugin interface. */
 class IPixWinPlugin : public IInputDeviceModule, public IRenderCaptureProvider
