Adding missing include for non-Editor builds.

diff --git a/Engine/Plugins/Developer/PixWinPlugin/Source/PixWinPlugin/Private/PixWinPluginModule.cpp b/Engine/Plugins/Developer/PixWinPlugin/Source/PixWinPlugin/Private/PixWinPluginModule.cpp
--- a/Engine/Plugins/Developer/PixWinPlugin/Source/PixWinPlugin/Private/PixWinPluginModule.cpp
+++ b/Engine/Plugins/Developer/PixWinPlugin/Source/PixWinPlugin/Private/PixWinPluginModule.cpp
@@ -8,6 +8,7 @@
 #include "Engine/GameViewportClient.h"
 #include "Engine/Engine.h"
 #include "Misc/CommandLine.h"
+#include "Misc/Paths.h"
 #include "HAL/FileManager.h"
 
 #if WITH_EDITOR
