Renaming PixWinPlugin to "PIX on Windows" since "PIX for Windows" was the name of the D3D9 PIX.

diff --git a/Engine/Plugins/Developer/PixWinPlugin/PixWinPlugin.uplugin b/Engine/Plugins/Developer/PixWinPlugin/PixWinPlugin.uplugin
--- a/Engine/Plugins/Developer/PixWinPlugin/PixWinPlugin.uplugin
+++ b/Engine/Plugins/Developer/PixWinPlugin/PixWinPlugin.uplugin
@@ -2,8 +2,8 @@
 	"FileVersion": 1,
 	"Version": 1,
 	"VersionName": "1.0",
-	"FriendlyName": "PIX for Windows GPU Capture Plugin",
-	"Description": "PIX for Windows graphics debugger integration.",
+	"FriendlyName": "PIX on Windows GPU Capture Plugin",
+	"Description": "PIX on Windows graphics debugger integration.",
 	"Category": "Rendering",
 	"CreatedBy": "Epic Games, Inc.",
 	"CreatedByURL": "https://epicgames.com",
