Fixed capture passes using RenderCaptureInterface::FScopedCapture being culled in RDG & made scoped captures able to open PIX. This path is used by e.g., r.CaptureNextDeferredShadingRendererFrame

diff --git a/Engine/Plugins/Developer/PixWinPlugin/Source/PixWinPlugin/Private/PixWinPluginModule.cpp b/Engine/Plugins/Developer/PixWinPlugin/Source/PixWinPlugin/Private/PixWinPluginModule.cpp
--- a/Engine/Plugins/Developer/PixWinPlugin/Source/PixWinPlugin/Private/PixWinPluginModule.cpp
+++ b/Engine/Plugins/Developer/PixWinPlugin/Source/PixWinPlugin/Private/PixWinPluginModule.cpp
@@ -219,10 +219,12 @@ void FPixWinPluginModule::CaptureFrame(FViewport* InViewport, uint32 InFlags, co
 void FPixWinPluginModule::BeginCapture(FRHICommandListImmediate* InRHICommandList, uint32 InFlags, const FString& InDestFileName)
 {
 #if PIX_PLUGIN_ENABLED
+	CurrentCaptureDestFileName = MakeWinPixCaptureFilePath(InDestFileName);
+	CurrentCaptureFlags = InFlags;
 	InRHICommandList->SubmitCommandsAndFlushGPU();
-	InRHICommandList->EnqueueLambda([Pix = PixGraphicsAnalysisInterface, InDestFileName](FRHICommandListImmediate& RHICommandList)
+	InRHICommandList->EnqueueLambda([Pix = PixGraphicsAnalysisInterface, DestFileName = CurrentCaptureDestFileName](FRHICommandListImmediate& RHICommandList)
 	{
-		Pix->BeginCapture(nullptr, InDestFileName);
+		Pix->BeginCapture(nullptr, DestFileName);
 	});
 #endif
 }
@@ -231,10 +233,21 @@ void FPixWinPluginModule::EndCapture(FRHICommandListImmediate* InRHICommandList)
 {
 #if PIX_PLUGIN_ENABLED
 	InRHICommandList->SubmitCommandsAndFlushGPU();
-	InRHICommandList->EnqueueLambda([Pix = PixGraphicsAnalysisInterface](FRHICommandListImmediate& RHICommandList)
+	InRHICommandList->EnqueueLambda(
+		[Pix = PixGraphicsAnalysisInterface, 
+		Flags = CurrentCaptureFlags, 
+		DestFileName = MoveTemp(CurrentCaptureDestFileName)](FRHICommandListImmediate& RHICommandList)
 	{
 		Pix->EndCapture();
+
+		// If we're already attached, don't open a new PIX instance
+		if ((Flags & IRenderCaptureProvider::ECaptureFlags_Launch) && !Pix->IsAttached())
+		{
+			Pix->OpenCapture(DestFileName);
+		}
 	});
+	CurrentCaptureFlags = 0u;
+	CurrentCaptureDestFileName.Empty();
 #endif
 }
 
diff --git a/Engine/Plugins/Developer/PixWinPlugin/Source/PixWinPlugin/Private/PixWinPluginModule.h b/Engine/Plugins/Developer/PixWinPlugin/Source/PixWinPlugin/Private/PixWinPluginModule.h
--- a/Engine/Plugins/Developer/PixWinPlugin/Source/PixWinPlugin/Private/PixWinPluginModule.h
+++ b/Engine/Plugins/Developer/PixWinPlugin/Source/PixWinPlugin/Private/PixWinPluginModule.h
@@ -58,6 +58,9 @@ private:
 
 	bool bCurrentlyAttached = false;
 
+	FString CurrentCaptureDestFileName;
+	uint32 CurrentCaptureFlags = 0u;
+
 #if WITH_EDITOR
 	TSharedPtr<FPixWinPluginEditorExtension> EditorExtension;
 #endif
diff --git a/Engine/Source/Runtime/RenderCore/Private/RenderCaptureInterface.cpp b/Engine/Source/Runtime/RenderCore/Private/RenderCaptureInterface.cpp
--- a/Engine/Source/Runtime/RenderCore/Private/RenderCaptureInterface.cpp
+++ b/Engine/Source/Runtime/RenderCore/Private/RenderCaptureInterface.cpp
@@ -76,7 +76,7 @@ namespace RenderCaptureInterface
 		{
 			GraphBuilder->AddPass(
 				RDG_EVENT_NAME("BeginCapture"),
-				ERDGPassFlags::None, 
+				ERDGPassFlags::NeverCull, 
 				[FileName = FString(InFileName)](FRHICommandListImmediate& RHICommandListLocal)
 				{
 					IRenderCaptureProvider::Get().BeginCapture(&RHICommandListLocal, IRenderCaptureProvider::ECaptureFlags_Launch, FString(FileName));
@@ -108,7 +108,7 @@ namespace RenderCaptureInterface
 
 				GraphBuilder->AddPass(
 					RDG_EVENT_NAME("EndCapture"), 
-					ERDGPassFlags::None, 
+					ERDGPassFlags::NeverCull, 
 					[](FRHICommandListImmediate& RHICommandListLocal)
 					{
 						IRenderCaptureProvider::Get().EndCapture(&RHICommandListLocal);
