Use a more specific timestamp in the PIX capture filename, to avoid subsequent captures overwriting previous ones.

#rnx

diff --git a/Engine/Plugins/Developer/PixWinPlugin/Source/PixWinPlugin/Private/PixWinPluginModule.cpp b/Engine/Plugins/Developer/PixWinPlugin/Source/PixWinPlugin/Private/PixWinPluginModule.cpp
--- a/Engine/Plugins/Developer/PixWinPlugin/Source/PixWinPlugin/Private/PixWinPluginModule.cpp
+++ b/Engine/Plugins/Developer/PixWinPlugin/Source/PixWinPlugin/Private/PixWinPluginModule.cpp
@@ -42,7 +42,7 @@ static FString MakeWinPixCaptureFilePath(const FString& InFilename)
 	if (Filename.IsEmpty())
 	{
 		const FDateTime DateTime = FDateTime::Now();
-		Filename = FString(TEXT("UEPixCapture_")) + DateTime.ToString();
+		Filename = FString(TEXT("UEPixCapture_")) + DateTime.ToString(TEXT("%Y.%m.%d-%H.%M.%S.%sZ"));
 	}
 
 	const bool bAbsoluteFileName = !FPaths::IsRelative(Filename);
