When taking a PIX capture in the editor, if there is no main viewport, look for alternate viewports.

#FailedValidations Preflight

diff --git a/Engine/Plugins/Developer/PixWinPlugin/Source/PixWinPlugin/Private/PixWinPluginModule.cpp b/Engine/Plugins/Developer/PixWinPlugin/Source/PixWinPlugin/Private/PixWinPluginModule.cpp
--- a/Engine/Plugins/Developer/PixWinPlugin/Source/PixWinPlugin/Private/PixWinPluginModule.cpp
+++ b/Engine/Plugins/Developer/PixWinPlugin/Source/PixWinPlugin/Private/PixWinPluginModule.cpp
@@ -14,6 +14,7 @@
 #if WITH_EDITOR
 #include "Framework/Application/SlateApplication.h"
 #include "Editor.h"
+#include "EditorViewportClient.h"
 #include "SPixWinPluginEditorExtension.h"
 #endif
 
@@ -318,6 +319,14 @@ void FPixWinPluginModule::DoFrameCaptureCurrentViewport(FViewport* InViewport, u
 		// (this has something to do with the 'active' editor viewport when the UI
 		// button is clicked versus the one which the console is attached to)
 		Viewport = GEditor->GetActiveViewport();
+		if (!Viewport)
+		{
+			const TArray<class FEditorViewportClient*>& ViewportClients = GEditor->GetAllViewportClients();
+			if (ViewportClients.Num() > 0)
+			{
+				Viewport = ViewportClients[0]->Viewport;
+			}
+		}
 	}
 #endif // WITH_EDITOR
 
